Fixed linker error

#rnx

diff --git a/Engine/Plugins/Runtime/AssetTags/Source/AssetTags/Private/AssetTagsSubsystem.cpp b/Engine/Plugins/Runtime/AssetTags/Source/AssetTags/Private/AssetTagsSubsystem.cpp
--- a/Engine/Plugins/Runtime/AssetTags/Source/AssetTags/Private/AssetTagsSubsystem.cpp
+++ b/Engine/Plugins/Runtime/AssetTags/Source/AssetTags/Private/AssetTagsSubsystem.cpp
@@ -324,7 +324,7 @@ bool UAssetTagsSubsystem::CollectionExists(const FName Name)
 
 		if (const FAssetRegistryState* StatePtr = AssetRegistry.GetAssetRegistryState())
 		{
-			const FString TagNameStr = FString::Printf(TEXT("%s%s"), FAssetData::CollectionTagPrefix, *Name.ToString());
+			const FString TagNameStr = FString::Printf(TEXT("%s%s"), FAssetData::GetCollectionTagPrefix(), *Name.ToString());
 			bExists = StatePtr->GetTagToAssetDatasMap().Contains(*TagNameStr);
 		}
 	}
@@ -356,12 +356,15 @@ TArray<FName> UAssetTagsSubsystem::GetCollections()
 
 		if (const FAssetRegistryState* StatePtr = AssetRegistry.GetAssetRegistryState())
 		{
+			const TCHAR* CollectionTagPrefix = FAssetData::GetCollectionTagPrefix();
+			const int32 CollectionTagPrefixLen = FCString::Strlen(CollectionTagPrefix);
+
 			for (const auto& TagAndArrayPair : StatePtr->GetTagToAssetDatasMap())
 			{
 				const FString TagNameStr = TagAndArrayPair.Key.ToString();
-				if (TagNameStr.StartsWith(FAssetData::CollectionTagPrefix))
+				if (TagNameStr.StartsWith(CollectionTagPrefix))
 				{
-					const FString TrimmedTagNameStr = TagNameStr.Mid(ARRAY_COUNT(FAssetData::CollectionTagPrefix) - 1);
+					const FString TrimmedTagNameStr = TagNameStr.Mid(CollectionTagPrefixLen);
 					CollectionNames.Add(*TrimmedTagNameStr);
 				}
 			}
@@ -401,7 +404,7 @@ TArray<FAssetData> UAssetTagsSubsystem::GetAssetsInCollection(const FName Name)
 	}
 #else	// WITH_EDITOR
 	{
-		const FString TagNameStr = FString::Printf(TEXT("%s%s"), FAssetData::CollectionTagPrefix, *Name.ToString());
+		const FString TagNameStr = FString::Printf(TEXT("%s%s"), FAssetData::GetCollectionTagPrefix(), *Name.ToString());
 
 		FARFilter Filter;
 		Filter.bIncludeOnlyOnDiskAssets = true; // Collection tags are added at cook-time, so we *must* search the on-disk version of the tags (from the asset registry)
@@ -440,12 +443,15 @@ TArray<FName> UAssetTagsSubsystem::GetCollectionsContainingAsset(const FName Ass
 		const FAssetData AssetData = AssetRegistry.GetAssetByObjectPath(AssetPathName, bIncludeOnlyOnDiskAssets);
 		if (AssetData.IsValid())
 		{
+			const TCHAR* CollectionTagPrefix = FAssetData::GetCollectionTagPrefix();
+			const int32 CollectionTagPrefixLen = FCString::Strlen(CollectionTagPrefix);
+
 			for (const auto& TagAndValuePair : AssetData.TagsAndValues)
 			{
 				const FString TagNameStr = TagAndValuePair.Key.ToString();
-				if (TagNameStr.StartsWith(FAssetData::CollectionTagPrefix))
+				if (TagNameStr.StartsWith(CollectionTagPrefix))
 				{
-					const FString TrimmedTagNameStr = TagNameStr.Mid(ARRAY_COUNT(FAssetData::CollectionTagPrefix) - 1);
+					const FString TrimmedTagNameStr = TagNameStr.Mid(CollectionTagPrefixLen);
 					CollectionNames.Add(*TrimmedTagNameStr);
 				}
 			}
diff --git a/Engine/Source/Editor/UnrealEd/Private/Commandlets/AssetRegistryGenerator.cpp b/Engine/Source/Editor/UnrealEd/Private/Commandlets/AssetRegistryGenerator.cpp
--- a/Engine/Source/Editor/UnrealEd/Private/Commandlets/AssetRegistryGenerator.cpp
+++ b/Engine/Source/Editor/UnrealEd/Private/Commandlets/AssetRegistryGenerator.cpp
@@ -747,7 +747,7 @@ void FAssetRegistryGenerator::UpdateCollectionAssetData()
 		TArray<FName> TmpAssetPathNames;
 		for (const FCollectionNameType& CollectionNameToTag : CollectionNamesToTag)
 		{
-			const FName CollectionTagName = *FString::Printf(TEXT("%s%s"), FAssetData::CollectionTagPrefix, *CollectionNameToTag.Name.ToString());
+			const FName CollectionTagName = *FString::Printf(TEXT("%s%s"), FAssetData::GetCollectionTagPrefix(), *CollectionNameToTag.Name.ToString());
 
 			TmpAssetPathNames.Reset();
 			CollectionManager.GetAssetsInCollection(CollectionNameToTag.Name, CollectionNameToTag.Type, TmpAssetPathNames);
diff --git a/Engine/Source/Runtime/AssetRegistry/Public/AssetData.h b/Engine/Source/Runtime/AssetRegistry/Public/AssetData.h
--- a/Engine/Source/Runtime/AssetRegistry/Public/AssetData.h
+++ b/Engine/Source/Runtime/AssetRegistry/Public/AssetData.h
@@ -57,7 +57,10 @@ struct FAssetData
 
 public:
 	/** The prefix used for collection entries inside TagsAndValues */
-	static constexpr TCHAR CollectionTagPrefix[] = TEXT("CL_");
+	static const TCHAR* GetCollectionTagPrefix()
+	{
+		return TEXT("CL_");
+	}
 
 public:
 	/** The object path for the asset in the form PackageName.AssetName. Only top level objects in a package can have AssetData */
