Add header guard for AssetTagsSubsystem.h

diff --git a/Engine/Plugins/Runtime/AssetTags/Source/AssetTags/Public/AssetTagsSubsystem.h b/Engine/Plugins/Runtime/AssetTags/Source/AssetTags/Public/AssetTagsSubsystem.h
--- a/Engine/Plugins/Runtime/AssetTags/Source/AssetTags/Public/AssetTagsSubsystem.h
+++ b/Engine/Plugins/Runtime/AssetTags/Source/AssetTags/Public/AssetTagsSubsystem.h
@@ -1,5 +1,7 @@
 // Copyright 1998-2019 Epic Games, Inc. All Rights Reserved.
 
+#pragma once
+
 #include "CoreMinimal.h"
 #include "Subsystems/EngineSubsystem.h"
 #include "AssetData.h"
