CollectionManagerScriptingSubsystem is a new subsystem exposing the Collection Container system to Blueprint.
-Adds blueprint nodes that offer feature parity with the blueprint API in AssetTagsSubsystem, but which allows the container to be specified.

[RN] CollectionManagerScriptingSubsystem is a new subsystem exposing the Collection Container system to Blueprint.

diff --git a/Engine/Plugins/Runtime/AssetTags/Source/AssetTags/AssetTags.Build.cs b/Engine/Plugins/Runtime/AssetTags/Source/AssetTags/AssetTags.Build.cs
--- a/Engine/Plugins/Runtime/AssetTags/Source/AssetTags/AssetTags.Build.cs
+++ b/Engine/Plugins/Runtime/AssetTags/Source/AssetTags/AssetTags.Build.cs
@@ -31,6 +31,7 @@ namespace UnrealBuildTool.Rules
 				PrivateDependencyModuleNames.AddRange(
 					new string[] {
 						"CollectionManager",
+						"UnrealEd"
 					}
 				);
 			}
diff --git a/Engine/Plugins/Runtime/AssetTags/Source/AssetTags/Private/AssetTagsSubsystem.cpp b/Engine/Plugins/Runtime/AssetTags/Source/AssetTags/Private/AssetTagsSubsystem.cpp
--- a/Engine/Plugins/Runtime/AssetTags/Source/AssetTags/Private/AssetTagsSubsystem.cpp
+++ b/Engine/Plugins/Runtime/AssetTags/Source/AssetTags/Private/AssetTagsSubsystem.cpp
@@ -8,24 +8,33 @@
 #include UE_INLINE_GENERATED_CPP_BY_NAME(AssetTagsSubsystem)
 
 #if WITH_EDITOR
+#include "CollectionManagerModule.h"
+#include "Editor/UnrealEdEngine.h"
 #include "ICollectionContainer.h"
 #include "ICollectionManager.h"
-#include "CollectionManagerModule.h"
+#include "ICollectionSource.h"
+#include "Subsystems/CollectionManagerScriptingSubsystem.h"
 #endif	// WITH_EDITOR
 
 DEFINE_LOG_CATEGORY_STATIC(LogAssetTags, Log, All);
 
 #if WITH_EDITOR
 
+extern UNREALED_API UEditorEngine* GEditor;
+
 namespace AssetTagsSubsystemUtil
 {
 
-void LogLastCollectionManagerError(const FText& Error)
+ECollectionScriptingShareType ECollectionShareType_To_ECollectionScriptingShareType(const ECollectionShareType::Type ShareType)
 {
-	UE_LOG(LogAssetTags, Warning, TEXT("Collection manager error: %s"), *Error.ToString());
+	static_assert((int32)ECollectionShareType::CST_Local == (int32)ECollectionScriptingShareType::Local + 1, "ECollectionShareType::CST_Local is expected to be ECollectionScriptingShareType::Local + 1");
+	static_assert((int32)ECollectionShareType::CST_Private == (int32)ECollectionScriptingShareType::Private + 1, "ECollectionShareType::CST_Private is expected to be ECollectionScriptingShareType::Private + 1");
+	static_assert((int32)ECollectionShareType::CST_Shared == (int32)ECollectionScriptingShareType::Shared + 1, "ECollectionShareType::CST_Shared is expected to be ECollectionScriptingShareType::Shared + 1");
+
+	return (ECollectionScriptingShareType)((int32)ShareType - 1);
 }
 
-FCollectionNameType FindCollectionByName(ICollectionContainer& CollectionContainer, const FName Name)
+bool FindCollectionByName(ICollectionContainer& CollectionContainer, const FName Name, FCollectionScriptingRef& OutCollection)
 {
 	TArray<FCollectionNameType> CollectionNamesAndTypes;
 	CollectionContainer.GetCollections(Name, CollectionNamesAndTypes);
@@ -33,190 +42,152 @@ FCollectionNameType FindCollectionByName(ICollectionContainer& CollectionContain
 	if (CollectionNamesAndTypes.Num() == 0)
 	{
 		UE_LOG(LogAssetTags, Warning, TEXT("No collection found called '%s'"), *Name.ToString());
-		return FCollectionNameType(NAME_None, ECollectionShareType::CST_All);
+		return false;
 	}
 	else if (CollectionNamesAndTypes.Num() > 1)
 	{
 		UE_LOG(LogAssetTags, Warning, TEXT("%d collections found called '%s'; ambiguous result"), CollectionNamesAndTypes.Num(), *Name.ToString());
-		return FCollectionNameType(NAME_None, ECollectionShareType::CST_All);
+		return false;
 	}
-	
-	return CollectionNamesAndTypes[0];
+
+	OutCollection = { .Container = CollectionContainer.GetCollectionSource()->GetName(), 
+		.Name = CollectionNamesAndTypes[0].Name, .ShareType = AssetTagsSubsystemUtil::ECollectionShareType_To_ECollectionScriptingShareType(CollectionNamesAndTypes[0].Type) };
+
+	return true;
 }
 
 }	// namespace AssetTagsSubsystemUtil
 
 bool UAssetTagsSubsystem::CreateCollection(const FName Name, const ECollectionScriptingShareType ShareType)
 {
-	const TSharedRef<ICollectionContainer>& CollectionContainer = FCollectionManagerModule::GetModule().Get().GetProjectCollectionContainer();
-
-	FText Error;
-	if (!CollectionContainer->IsValidCollectionName(Name.ToString(), ECollectionShareType::CST_All, &Error))
-	{
-		AssetTagsSubsystemUtil::LogLastCollectionManagerError(Error);
-		return false;
-	}
-
-	static_assert((int32)ECollectionShareType::CST_Local == (int32)ECollectionScriptingShareType::Local + 1, "ECollectionShareType::CST_Local is expected to be ECollectionScriptingShareType::Local + 1");
-	static_assert((int32)ECollectionShareType::CST_Private == (int32)ECollectionScriptingShareType::Private + 1, "ECollectionShareType::CST_Private is expected to be ECollectionScriptingShareType::Private + 1");
-	static_assert((int32)ECollectionShareType::CST_Shared == (int32)ECollectionScriptingShareType::Shared + 1, "ECollectionShareType::CST_Shared is expected to be ECollectionScriptingShareType::Shared + 1");
-
-	if (!CollectionContainer->CreateCollection(Name, (ECollectionShareType::Type)((int32)ShareType + 1), ECollectionStorageMode::Static, &Error))
+	if (UCollectionManagerScriptingSubsystem* Subsystem = GEditor->GetEditorSubsystem<UCollectionManagerScriptingSubsystem>(); ensure(Subsystem))
 	{
-		AssetTagsSubsystemUtil::LogLastCollectionManagerError(Error);
-		return false;
+		const TSharedRef<ICollectionContainer>& CollectionContainer = FCollectionManagerModule::GetModule().Get().GetProjectCollectionContainer();
+		FCollectionScriptingRef NewCollection;
+		return Subsystem->CreateCollection(FCollectionScriptingContainerSource{ CollectionContainer->GetCollectionSource()->GetName(), CollectionContainer->GetCollectionSource()->GetTitle() }, Name, ShareType, NewCollection);
 	}
-
-	return true;
+	
+	return false;
 }
 
 bool UAssetTagsSubsystem::DestroyCollection(const FName Name)
 {
-	const TSharedRef<ICollectionContainer>& CollectionContainer = FCollectionManagerModule::GetModule().Get().GetProjectCollectionContainer();
-
-	const FCollectionNameType ResolvedNameAndType = AssetTagsSubsystemUtil::FindCollectionByName(*CollectionContainer, Name);
-	if (ResolvedNameAndType.Name.IsNone())
+	if (UCollectionManagerScriptingSubsystem* Subsystem = GEditor->GetEditorSubsystem<UCollectionManagerScriptingSubsystem>(); ensure(Subsystem))
 	{
-		return false;
-	}
-
-	FText Error;
-	if (!CollectionContainer->DestroyCollection(ResolvedNameAndType.Name, ResolvedNameAndType.Type, &Error))
-	{
-		AssetTagsSubsystemUtil::LogLastCollectionManagerError(Error);
-		return false;
+		const TSharedRef<ICollectionContainer>& CollectionContainer = FCollectionManagerModule::GetModule().Get().GetProjectCollectionContainer();
+		FCollectionScriptingRef Collection;
+		if (AssetTagsSubsystemUtil::FindCollectionByName(*CollectionContainer, Name, Collection))
+		{
+			return Subsystem->DestroyCollection(Collection);
+		}
 	}
 
-	return true;
+	return false;
 }
 
 bool UAssetTagsSubsystem::RenameCollection(const FName Name, const FName NewName)
 {
-	const TSharedRef<ICollectionContainer>& CollectionContainer = FCollectionManagerModule::GetModule().Get().GetProjectCollectionContainer();
-
-	const FCollectionNameType ResolvedNameAndType = AssetTagsSubsystemUtil::FindCollectionByName(*CollectionContainer, Name);
-	if (ResolvedNameAndType.Name.IsNone())
+	if (UCollectionManagerScriptingSubsystem* Subsystem = GEditor->GetEditorSubsystem<UCollectionManagerScriptingSubsystem>(); ensure(Subsystem))
 	{
-		return false;
-	}
-
-	FText Error;
-	if (!CollectionContainer->IsValidCollectionName(NewName.ToString(), ECollectionShareType::CST_All, &Error))
-	{
-		AssetTagsSubsystemUtil::LogLastCollectionManagerError(Error);
-		return false;
-	}
+		const TSharedRef<ICollectionContainer>& CollectionContainer = FCollectionManagerModule::GetModule().Get().GetProjectCollectionContainer();
 
-	if (!CollectionContainer->RenameCollection(ResolvedNameAndType.Name, ResolvedNameAndType.Type, NewName, ResolvedNameAndType.Type, &Error))
-	{
-		AssetTagsSubsystemUtil::LogLastCollectionManagerError(Error);
-		return false;
+		FCollectionScriptingRef Collection;
+		if (AssetTagsSubsystemUtil::FindCollectionByName(*CollectionContainer, Name, Collection))
+		{
+			return Subsystem->RenameCollection(Collection, NewName, Collection.ShareType);
+		}
 	}
 
-	return true;
+	return false;
 }
 
 bool UAssetTagsSubsystem::ReparentCollection(const FName Name, const FName NewParentName)
 {
-	const TSharedRef<ICollectionContainer>& CollectionContainer = FCollectionManagerModule::GetModule().Get().GetProjectCollectionContainer();
-
-	const FCollectionNameType ResolvedNameAndType = AssetTagsSubsystemUtil::FindCollectionByName(*CollectionContainer, Name);
-	if (ResolvedNameAndType.Name.IsNone())
+	if (UCollectionManagerScriptingSubsystem* Subsystem = GEditor->GetEditorSubsystem<UCollectionManagerScriptingSubsystem>(); ensure(Subsystem))
 	{
-		return false;
-	}
-
-	const FCollectionNameType ResolvedParentNameAndType = AssetTagsSubsystemUtil::FindCollectionByName(*CollectionContainer, NewParentName);
-	if (ResolvedParentNameAndType.Name.IsNone())
-	{
-		return false;
-	}
+		const TSharedRef<ICollectionContainer>& CollectionContainer = FCollectionManagerModule::GetModule().Get().GetProjectCollectionContainer();
 
-	FText Error;
-	if (!CollectionContainer->ReparentCollection(ResolvedNameAndType.Name, ResolvedNameAndType.Type, ResolvedParentNameAndType.Name, ResolvedParentNameAndType.Type, &Error))
-	{
-		AssetTagsSubsystemUtil::LogLastCollectionManagerError(Error);
-		return false;
+		FCollectionScriptingRef ResolvedCollection;
+		if (AssetTagsSubsystemUtil::FindCollectionByName(*CollectionContainer, Name, ResolvedCollection))
+		{
+			FCollectionScriptingRef ResolvedParentCollection { .Container = CollectionContainer->GetCollectionSource()->GetName() };
+			if (NewParentName.IsNone() || AssetTagsSubsystemUtil::FindCollectionByName(*CollectionContainer, NewParentName, ResolvedParentCollection))
+			{
+				return Subsystem->ReparentCollection(ResolvedCollection, ResolvedParentCollection);
+			}
+		}
 	}
 
-	return true;
+	return false;
 }
 
 bool UAssetTagsSubsystem::EmptyCollection(const FName Name)
 {
-	const TSharedRef<ICollectionContainer>& CollectionContainer = FCollectionManagerModule::GetModule().Get().GetProjectCollectionContainer();
-
-	const FCollectionNameType ResolvedNameAndType = AssetTagsSubsystemUtil::FindCollectionByName(*CollectionContainer, Name);
-	if (ResolvedNameAndType.Name.IsNone())
+	if (UCollectionManagerScriptingSubsystem* Subsystem = GEditor->GetEditorSubsystem<UCollectionManagerScriptingSubsystem>(); ensure(Subsystem))
 	{
-		return false;
-	}
+		const TSharedRef<ICollectionContainer>& CollectionContainer = FCollectionManagerModule::GetModule().Get().GetProjectCollectionContainer();
 
-	FText Error;
-	if (!CollectionContainer->EmptyCollection(ResolvedNameAndType.Name, ResolvedNameAndType.Type, &Error))
-	{
-		AssetTagsSubsystemUtil::LogLastCollectionManagerError(Error);
-		return false;
+		FCollectionScriptingRef ResolvedCollection;
+		if (AssetTagsSubsystemUtil::FindCollectionByName(*CollectionContainer, Name, ResolvedCollection))
+		{
+			return Subsystem->EmptyCollection(ResolvedCollection);
+		}
 	}
 
-	return true;
+	return false;
 }
 
 bool UAssetTagsSubsystem::K2_AddAssetToCollection(const FName Name, const FSoftObjectPath& AssetPath)
 {
-	const TSharedRef<ICollectionContainer>& CollectionContainer = FCollectionManagerModule::GetModule().Get().GetProjectCollectionContainer();
-
-	const FCollectionNameType ResolvedNameAndType = AssetTagsSubsystemUtil::FindCollectionByName(*CollectionContainer, Name);
-	if (ResolvedNameAndType.Name.IsNone())
+	if (UCollectionManagerScriptingSubsystem* Subsystem = GEditor->GetEditorSubsystem<UCollectionManagerScriptingSubsystem>(); ensure(Subsystem))
 	{
-		return false;
-	}
+		const TSharedRef<ICollectionContainer>& CollectionContainer = FCollectionManagerModule::GetModule().Get().GetProjectCollectionContainer();
 
-	FText Error;
-	if (!CollectionContainer->AddToCollection(ResolvedNameAndType.Name, ResolvedNameAndType.Type, AssetPath, &Error))
-	{
-		AssetTagsSubsystemUtil::LogLastCollectionManagerError(Error);
-		return false;
+		FCollectionScriptingRef ResolvedCollection;
+		if (AssetTagsSubsystemUtil::FindCollectionByName(*CollectionContainer, Name, ResolvedCollection))
+		{
+			return Subsystem->AddAssetToCollection(ResolvedCollection, AssetPath);
+		}
 	}
 
-	return true;
+	return false;
 }
 
 bool UAssetTagsSubsystem::AddAssetToCollection(const FName Name, const FName AssetPathName)
 {
+PRAGMA_DISABLE_DEPRECATION_WARNINGS
 	return K2_AddAssetToCollection(Name, FSoftObjectPath(AssetPathName.ToString()));
+PRAGMA_ENABLE_DEPRECATION_WARNINGS
 }
 
 bool UAssetTagsSubsystem::AddAssetDataToCollection(const FName Name, const FAssetData& AssetData)
 {
+PRAGMA_DISABLE_DEPRECATION_WARNINGS
 	return K2_AddAssetToCollection(Name, AssetData.GetSoftObjectPath());
+PRAGMA_ENABLE_DEPRECATION_WARNINGS
 }
 
 bool UAssetTagsSubsystem::AddAssetPtrToCollection(const FName Name, const UObject* AssetPtr)
 {
 PRAGMA_DISABLE_DEPRECATION_WARNINGS
-	return AddAssetToCollection(Name, AssetPtr ? FName(*AssetPtr->GetPathName()) : NAME_None);
+	return K2_AddAssetToCollection(Name, FSoftObjectPath(AssetPtr));
 PRAGMA_ENABLE_DEPRECATION_WARNINGS
 }
 
 bool UAssetTagsSubsystem::K2_AddAssetsToCollection(const FName Name, const TArray<FSoftObjectPath>& AssetPaths)
 {
-	const TSharedRef<ICollectionContainer>& CollectionContainer = FCollectionManagerModule::GetModule().Get().GetProjectCollectionContainer();
-
-	const FCollectionNameType ResolvedNameAndType = AssetTagsSubsystemUtil::FindCollectionByName(*CollectionContainer, Name);
-	if (ResolvedNameAndType.Name.IsNone())
+	if (UCollectionManagerScriptingSubsystem* Subsystem = GEditor->GetEditorSubsystem<UCollectionManagerScriptingSubsystem>(); ensure(Subsystem))
 	{
-		return false;
-	}
+		const TSharedRef<ICollectionContainer>& CollectionContainer = FCollectionManagerModule::GetModule().Get().GetProjectCollectionContainer();
 
-	FText Error;
-	if (!CollectionContainer->AddToCollection(ResolvedNameAndType.Name, ResolvedNameAndType.Type, AssetPaths, nullptr, &Error))
-	{
-		AssetTagsSubsystemUtil::LogLastCollectionManagerError(Error);
-		return false;
+		FCollectionScriptingRef ResolvedCollection;
+		if (AssetTagsSubsystemUtil::FindCollectionByName(*CollectionContainer, Name, ResolvedCollection))
+		{
+			return Subsystem->AddAssetsToCollection(ResolvedCollection, AssetPaths);
+		}
 	}
 
-	return true;
+	return false;
 }
 
 bool UAssetTagsSubsystem::AddAssetsToCollection(const FName Name, const TArray<FName>& AssetPathNames)
@@ -234,7 +205,10 @@ bool UAssetTagsSubsystem::AddAssetDatasToCollection(const FName Name, const TArr
 	{
 		AssetPaths.Add(AssetData.GetSoftObjectPath());
 	}
+
+PRAGMA_DISABLE_DEPRECATION_WARNINGS
 	return K2_AddAssetsToCollection(Name, AssetPaths);
+PRAGMA_ENABLE_DEPRECATION_WARNINGS
 }
 
 bool UAssetTagsSubsystem::AddAssetPtrsToCollection(const FName Name, const TArray<UObject*>& AssetPtrs)
@@ -245,64 +219,63 @@ bool UAssetTagsSubsystem::AddAssetPtrsToCollection(const FName Name, const TArra
 	{
 		AssetPaths.Add(FSoftObjectPath(AssetPtr));
 	}
+
+PRAGMA_DISABLE_DEPRECATION_WARNINGS
 	return K2_AddAssetsToCollection(Name, AssetPaths);
+PRAGMA_ENABLE_DEPRECATION_WARNINGS
 }
 
 bool UAssetTagsSubsystem::K2_RemoveAssetFromCollection(const FName Name, const FSoftObjectPath& AssetPath)
 {
-	const TSharedRef<ICollectionContainer>& CollectionContainer = FCollectionManagerModule::GetModule().Get().GetProjectCollectionContainer();
-
-	const FCollectionNameType ResolvedNameAndType = AssetTagsSubsystemUtil::FindCollectionByName(*CollectionContainer, Name);
-	if (ResolvedNameAndType.Name.IsNone())
+	if (UCollectionManagerScriptingSubsystem* Subsystem = GEditor->GetEditorSubsystem<UCollectionManagerScriptingSubsystem>(); ensure(Subsystem))
 	{
-		return false;
-	}
+		const TSharedRef<ICollectionContainer>& CollectionContainer = FCollectionManagerModule::GetModule().Get().GetProjectCollectionContainer();
 
-	FText Error;
-	if (!CollectionContainer->RemoveFromCollection(ResolvedNameAndType.Name, ResolvedNameAndType.Type, AssetPath, &Error))
-	{
-		AssetTagsSubsystemUtil::LogLastCollectionManagerError(Error);
-		return false;
+		FCollectionScriptingRef ResolvedCollection;
+		if (AssetTagsSubsystemUtil::FindCollectionByName(*CollectionContainer, Name, ResolvedCollection))
+		{
+			return Subsystem->RemoveAssetFromCollection(ResolvedCollection, AssetPath);
+		}
 	}
 
-	return true;
+	return false;
 }
 
 bool UAssetTagsSubsystem::RemoveAssetFromCollection(const FName Name, const FName AssetPathName)
 {
+PRAGMA_DISABLE_DEPRECATION_WARNINGS
 	return K2_RemoveAssetFromCollection(Name, FSoftObjectPath(AssetPathName.ToString()));
+PRAGMA_ENABLE_DEPRECATION_WARNINGS
 }
 
 bool UAssetTagsSubsystem::RemoveAssetDataFromCollection(const FName Name, const FAssetData& AssetData)
 {
+PRAGMA_DISABLE_DEPRECATION_WARNINGS
 	return K2_RemoveAssetFromCollection(Name, AssetData.GetSoftObjectPath());
+PRAGMA_ENABLE_DEPRECATION_WARNINGS
 }
 
 bool UAssetTagsSubsystem::RemoveAssetPtrFromCollection(const FName Name, const UObject* AssetPtr)
 {
 PRAGMA_DISABLE_DEPRECATION_WARNINGS
-	return RemoveAssetFromCollection(Name, AssetPtr ? FName(*AssetPtr->GetPathName()) : NAME_None);
+	return K2_RemoveAssetFromCollection(Name, FSoftObjectPath(AssetPtr));
 PRAGMA_ENABLE_DEPRECATION_WARNINGS
 }
 
 bool UAssetTagsSubsystem::K2_RemoveAssetsFromCollection(const FName Name, const TArray<FSoftObjectPath>& AssetPaths)
 {
-	const TSharedRef<ICollectionContainer>& CollectionContainer = FCollectionManagerModule::GetModule().Get().GetProjectCollectionContainer();
-
-	const FCollectionNameType ResolvedNameAndType = AssetTagsSubsystemUtil::FindCollectionByName(*CollectionContainer, Name);
-	if (ResolvedNameAndType.Name.IsNone())
+	if (UCollectionManagerScriptingSubsystem* Subsystem = GEditor->GetEditorSubsystem<UCollectionManagerScriptingSubsystem>(); ensure(Subsystem))
 	{
-		return false;
-	}
+		const TSharedRef<ICollectionContainer>& CollectionContainer = FCollectionManagerModule::GetModule().Get().GetProjectCollectionContainer();
 
-	FText Error;
-	if (!CollectionContainer->RemoveFromCollection(ResolvedNameAndType.Name, ResolvedNameAndType.Type, AssetPaths, nullptr, &Error))
-	{
-		AssetTagsSubsystemUtil::LogLastCollectionManagerError(Error);
-		return false;
+		FCollectionScriptingRef ResolvedCollection;
+		if (AssetTagsSubsystemUtil::FindCollectionByName(*CollectionContainer, Name, ResolvedCollection))
+		{
+			return Subsystem->RemoveAssetsFromCollection(ResolvedCollection, AssetPaths);
+		}
 	}
 
-	return true;
+	return false;
 }
 
 bool UAssetTagsSubsystem::RemoveAssetsFromCollection(const FName Name, const TArray<FName>& AssetPathNames)
@@ -320,7 +293,9 @@ bool UAssetTagsSubsystem::RemoveAssetDatasFromCollection(const FName Name, const
 	{
 		AssetPaths.Add(AssetData.GetSoftObjectPath());
 	}
+PRAGMA_DISABLE_DEPRECATION_WARNINGS
 	return K2_RemoveAssetsFromCollection(Name, AssetPaths);
+PRAGMA_ENABLE_DEPRECATION_WARNINGS
 }
 
 bool UAssetTagsSubsystem::RemoveAssetPtrsFromCollection(const FName Name, const TArray<UObject*>& AssetPtrs)
@@ -331,7 +306,9 @@ bool UAssetTagsSubsystem::RemoveAssetPtrsFromCollection(const FName Name, const
 	{
 		AssetPaths.Add(FSoftObjectPath(AssetPtr));
 	}
+PRAGMA_DISABLE_DEPRECATION_WARNINGS
 	return K2_RemoveAssetsFromCollection(Name, AssetPaths);
+PRAGMA_ENABLE_DEPRECATION_WARNINGS
 }
 
 #endif	// WITH_EDITOR
@@ -342,8 +319,11 @@ bool UAssetTagsSubsystem::CollectionExists(const FName Name)
 
 #if WITH_EDITOR
 	{
-		const TSharedRef<ICollectionContainer>& CollectionContainer = FCollectionManagerModule::GetModule().Get().GetProjectCollectionContainer();
-		bExists = CollectionContainer->CollectionExists(Name, ECollectionShareType::CST_All);
+		if (UCollectionManagerScriptingSubsystem* Subsystem = GEditor->GetEditorSubsystem<UCollectionManagerScriptingSubsystem>(); ensure(Subsystem))
+		{
+			const TSharedRef<ICollectionContainer>& CollectionContainer = FCollectionManagerModule::GetModule().Get().GetProjectCollectionContainer();
+			bExists = Subsystem->CollectionExistsWithAnyShareType(FCollectionScriptingContainerSource{ CollectionContainer->GetCollectionSource()->GetName(), CollectionContainer->GetCollectionSource()->GetTitle() }, Name);
+		}
 	}
 #else	// WITH_EDITOR
 	{
@@ -363,15 +343,19 @@ TArray<FName> UAssetTagsSubsystem::GetCollections()
 
 #if WITH_EDITOR
 	{
-		const TSharedRef<ICollectionContainer>& CollectionContainer = FCollectionManagerModule::GetModule().Get().GetProjectCollectionContainer();
-
-		TArray<FCollectionNameType> CollectionNamesAndTypes;
-		CollectionContainer->GetCollections(CollectionNamesAndTypes);
-
-		CollectionNames.Reserve(CollectionNamesAndTypes.Num());
-		for (const FCollectionNameType& CollectionNameAndType : CollectionNamesAndTypes)
+		if (UCollectionManagerScriptingSubsystem* Subsystem = GEditor->GetEditorSubsystem<UCollectionManagerScriptingSubsystem>(); ensure(Subsystem))
 		{
-			CollectionNames.AddUnique(CollectionNameAndType.Name);
+			const TSharedRef<ICollectionContainer>& CollectionContainer = FCollectionManagerModule::GetModule().Get().GetProjectCollectionContainer();
+
+			TArray<FCollectionScriptingRef> Collections;
+			if (Subsystem->GetCollections(FCollectionScriptingContainerSource{ CollectionContainer->GetCollectionSource()->GetName(), CollectionContainer->GetCollectionSource()->GetTitle() }, Collections))
+			{
+				CollectionNames.Reserve(Collections.Num());
+				for (const FCollectionScriptingRef& Collection : Collections)
+				{
+					CollectionNames.AddUnique(Collection.Name);
+				}
+			}
 		}
 	}
 #else	// WITH_EDITOR
@@ -407,21 +391,14 @@ TArray<FAssetData> UAssetTagsSubsystem::GetAssetsInCollection(const FName Name)
 	IAssetRegistry& AssetRegistry = FModuleManager::LoadModuleChecked<FAssetRegistryModule>(AssetRegistryConstants::ModuleName).Get();
 #if WITH_EDITOR
 	{
-		const TSharedRef<ICollectionContainer>& CollectionContainer = FCollectionManagerModule::GetModule().Get().GetProjectCollectionContainer();
-
-		const FCollectionNameType ResolvedNameAndType = AssetTagsSubsystemUtil::FindCollectionByName(*CollectionContainer, Name);
-		if (!ResolvedNameAndType.Name.IsNone())
+		if (UCollectionManagerScriptingSubsystem* Subsystem = GEditor->GetEditorSubsystem<UCollectionManagerScriptingSubsystem>(); ensure(Subsystem))
 		{
-			TArray<FSoftObjectPath> AssetPaths;
-			CollectionContainer->GetAssetsInCollection(ResolvedNameAndType.Name, ResolvedNameAndType.Type, AssetPaths);
-			Assets.Reserve(AssetPaths.Num());
-			for (const FSoftObjectPath& AssetPath : AssetPaths)
+			const TSharedRef<ICollectionContainer>& CollectionContainer = FCollectionManagerModule::GetModule().Get().GetProjectCollectionContainer();
+
+			FCollectionScriptingRef ResolvedCollection;
+			if (AssetTagsSubsystemUtil::FindCollectionByName(*CollectionContainer, Name, ResolvedCollection))
 			{
-				FAssetData AssetData = AssetRegistry.GetAssetByObjectPath(AssetPath);
-				if (AssetData.IsValid())
-				{
-					Assets.Add(MoveTemp(AssetData));
-				}
+				Subsystem->GetAssetsInCollection(ResolvedCollection, Assets);
 			}
 		}
 	}
@@ -447,15 +424,19 @@ TArray<FName> UAssetTagsSubsystem::K2_GetCollectionsContainingAsset(const FSoftO
 
 #if WITH_EDITOR
 	{
-		const TSharedRef<ICollectionContainer>& CollectionContainer = FCollectionManagerModule::GetModule().Get().GetProjectCollectionContainer();
-
-		TArray<FCollectionNameType> CollectionNamesAndTypes;
-		CollectionContainer->GetCollectionsContainingObject(AssetPath, CollectionNamesAndTypes);
-
-		CollectionNames.Reserve(CollectionNamesAndTypes.Num());
-		for (const FCollectionNameType& CollectionNameAndType : CollectionNamesAndTypes)
+		if (UCollectionManagerScriptingSubsystem* Subsystem = GEditor->GetEditorSubsystem<UCollectionManagerScriptingSubsystem>(); ensure(Subsystem))
 		{
-			CollectionNames.AddUnique(CollectionNameAndType.Name);
+			const TSharedRef<ICollectionContainer>& CollectionContainer = FCollectionManagerModule::GetModule().Get().GetProjectCollectionContainer();
+
+			TArray<FCollectionScriptingRef> Collections;
+			if (Subsystem->GetCollectionsContainingAsset(FCollectionScriptingContainerSource{ CollectionContainer->GetCollectionSource()->GetName(), CollectionContainer->GetCollectionSource()->GetTitle() }, AssetPath, Collections))
+			{
+				CollectionNames.Reserve(Collections.Num());
+				for (const FCollectionScriptingRef& Collection : Collections)
+				{
+					CollectionNames.Add(Collection.Name);
+				}
+			}
 		}
 	}
 #else	// WITH_EDITOR
diff --git a/Engine/Plugins/Runtime/AssetTags/Source/AssetTags/Public/AssetTagsSubsystem.h b/Engine/Plugins/Runtime/AssetTags/Source/AssetTags/Public/AssetTagsSubsystem.h
--- a/Engine/Plugins/Runtime/AssetTags/Source/AssetTags/Public/AssetTagsSubsystem.h
+++ b/Engine/Plugins/Runtime/AssetTags/Source/AssetTags/Public/AssetTagsSubsystem.h
@@ -2,21 +2,11 @@
 
 #pragma once
 
+#include "CollectionManagerScriptingTypes.h"
 #include "Subsystems/EngineSubsystem.h"
 #include "UObject/Package.h"
 #include "AssetTagsSubsystem.generated.h"
 
-UENUM(DisplayName="Collection Share Type", meta=(ScriptName="CollectionShareType"))
-enum class ECollectionScriptingShareType : uint8
-{
-	/** This collection is only visible to you and is not in source control. */
-	Local,
-	/** This collection is only visible to you. */
-	Private,
-	/** This collection is visible to everyone. */
-	Shared,
-};
-
 UCLASS()
 class ASSETTAGS_API UAssetTagsSubsystem : public UEngineSubsystem
 {
diff --git a/Engine/Source/Editor/UnrealEd/Private/Subsystems/CollectionManagerScriptingSubsystem.cpp b/Engine/Source/Editor/UnrealEd/Private/Subsystems/CollectionManagerScriptingSubsystem.cpp
new file mode 100644
--- /dev/null
+++ b/Engine/Source/Editor/UnrealEd/Private/Subsystems/CollectionManagerScriptingSubsystem.cpp
@@ -0,0 +1,459 @@
+// Copyright Epic Games, Inc. All Rights Reserved.
+
+#include "Subsystems/CollectionManagerScriptingSubsystem.h"
+
+#include "AssetRegistry/AssetData.h"
+#include "AssetRegistry/AssetRegistryModule.h"
+#include "CollectionManagerModule.h"
+#include "ICollectionContainer.h"
+#include "ICollectionManager.h"
+#include "ICollectionSource.h"
+
+#include UE_INLINE_GENERATED_CPP_BY_NAME(CollectionManagerScriptingSubsystem)
+
+DEFINE_LOG_CATEGORY_STATIC(LogCollectionManagerScripting, Log, All);
+
+#define LOCTEXT_NAMESPACE "CollectionManagerScriptingSubsystem"
+
+namespace CollectionManagerScriptingSubsystemUtil
+{
+
+void LogLastCollectionManagerError(const FText& Error)
+{
+	UE_LOG(LogCollectionManagerScripting, Warning, TEXT("Collection manager error: %s"), *Error.ToString());
+}
+
+FString ECollectionScriptingShareType_To_String(const ECollectionScriptingShareType ShareType)
+{
+	static const UEnum* Enum = FindObject<UEnum>(nullptr, TEXT("/Script/CollectionManager.ECollectionScriptingShareType"), true);
+	if (Enum)
+	{
+		return Enum->GetDisplayNameTextByValue((uint64)ShareType).ToString();
+	}
+	else
+	{
+		return {};
+	}
+}
+
+ECollectionShareType::Type ECollectionScriptingShareType_To_ECollectionShareType(const ECollectionScriptingShareType ShareType)
+{
+	static_assert((int32)ECollectionShareType::CST_Local == (int32)ECollectionScriptingShareType::Local + 1, "ECollectionShareType::CST_Local is expected to be ECollectionScriptingShareType::Local + 1");
+	static_assert((int32)ECollectionShareType::CST_Private == (int32)ECollectionScriptingShareType::Private + 1, "ECollectionShareType::CST_Private is expected to be ECollectionScriptingShareType::Private + 1");
+	static_assert((int32)ECollectionShareType::CST_Shared == (int32)ECollectionScriptingShareType::Shared + 1, "ECollectionShareType::CST_Shared is expected to be ECollectionScriptingShareType::Shared + 1");
+
+	return (ECollectionShareType::Type)((int32)ShareType + 1);
+}
+
+ECollectionScriptingShareType ECollectionShareType_To_ECollectionScriptingShareType(const ECollectionShareType::Type ShareType)
+{
+	static_assert((int32)ECollectionShareType::CST_Local == (int32)ECollectionScriptingShareType::Local + 1, "ECollectionShareType::CST_Local is expected to be ECollectionScriptingShareType::Local + 1");
+	static_assert((int32)ECollectionShareType::CST_Private == (int32)ECollectionScriptingShareType::Private + 1, "ECollectionShareType::CST_Private is expected to be ECollectionScriptingShareType::Private + 1");
+	static_assert((int32)ECollectionShareType::CST_Shared == (int32)ECollectionScriptingShareType::Shared + 1, "ECollectionShareType::CST_Shared is expected to be ECollectionScriptingShareType::Shared + 1");
+
+	return (ECollectionScriptingShareType)((int32)ShareType - 1);
+}
+
+struct FConstants
+{
+	static const FConstants& Get()
+	{
+		static FConstants Instance;
+		return Instance;
+	}
+
+	const FText ContainerNotFound;
+	const FText CollectionNotFound;
+	const FText CollectionExistsInContainer;
+	const FText ContainerMismatch;
+	const FText FailedToGetAsset;
+
+private:
+	FConstants() :
+		ContainerNotFound(LOCTEXT("ContainerNotFound", "No container with the name '{0}' could be found.")),
+		CollectionNotFound(LOCTEXT("CollectionNotFound", "No collection with the name '{0}' and share type '{1}' was found in container with name '{2}'")),
+		CollectionExistsInContainer(LOCTEXT("CollectionExistsInContainer", "Container '{0}' already has a collection named '{1}' with share type '{2}'.")),
+		ContainerMismatch(LOCTEXT("ContainerMismatch", "Collections must be in the same container to reparent. Found target collection in container '{0}' and new parent collection in container '{1}'.")),
+		FailedToGetAsset(LOCTEXT("FailedToGetAsset", "Failed to get Asset '{0}' from Collection '{1}'. Could not get the asset's Object Path."))
+	{}
+};
+
+} // namespace CollectionManagerScriptingSubsystemUtil
+
+TArray<FCollectionScriptingContainerSource> UCollectionManagerScriptingSubsystem::GetCollectionContainers()
+{
+	TArray<TSharedPtr<ICollectionContainer>> CollectionContainers;
+	FCollectionManagerModule::GetModule().Get().GetCollectionContainers(CollectionContainers);
+
+	TArray<FCollectionScriptingContainerSource> ContainerNames;
+	ContainerNames.Reserve(CollectionContainers.Num());
+	Algo::Transform(CollectionContainers, ContainerNames, [](const TSharedPtr<ICollectionContainer>& CollectionContainer)
+		{
+			return FCollectionScriptingContainerSource {
+				.Name = CollectionContainer->GetCollectionSource()->GetName(),
+				.Title = CollectionContainer->GetCollectionSource()->GetTitle()
+			};
+		}
+	);
+
+	return ContainerNames;
+}
+
+bool UCollectionManagerScriptingSubsystem::CreateCollection(const FCollectionScriptingContainerSource Container, const FName Collection, const ECollectionScriptingShareType ShareType, FCollectionScriptingRef& OutNewCollection)
+{
+	const TSharedPtr<ICollectionContainer> CollectionContainer = FCollectionManagerModule::GetModule().Get().FindCollectionContainer(Container.Name);
+	if (!CollectionContainer.IsValid())
+	{
+		CollectionManagerScriptingSubsystemUtil::LogLastCollectionManagerError(FText::Format(CollectionManagerScriptingSubsystemUtil::FConstants::Get().ContainerNotFound, FText::FromName(Container.Name)));
+		return false;
+	}
+
+	if (const ECollectionShareType::Type CastShareType = CollectionManagerScriptingSubsystemUtil::ECollectionScriptingShareType_To_ECollectionShareType(ShareType); CollectionContainer->CollectionExists(Collection, CastShareType))
+	{
+		CollectionManagerScriptingSubsystemUtil::LogLastCollectionManagerError(FText::Format(CollectionManagerScriptingSubsystemUtil::FConstants::Get().CollectionExistsInContainer, { CollectionContainer->GetCollectionSource()->GetTitle(), FText::FromName(Collection), CastShareType }));
+		return false;
+	}
+
+	FText Error;
+	if (!CollectionContainer->IsValidCollectionName(Collection.ToString(), ECollectionShareType::CST_All, &Error))
+	{
+		CollectionManagerScriptingSubsystemUtil::LogLastCollectionManagerError(Error);
+		return false;
+	}
+
+	if (!CollectionContainer->CreateCollection(Collection, CollectionManagerScriptingSubsystemUtil::ECollectionScriptingShareType_To_ECollectionShareType(ShareType), ECollectionStorageMode::Static, &Error))
+	{
+		CollectionManagerScriptingSubsystemUtil::LogLastCollectionManagerError(Error);
+		return false;
+	}
+
+	OutNewCollection.Name = Collection;
+	OutNewCollection.ShareType = ShareType;
+	OutNewCollection.Container = CollectionContainer->GetCollectionSource()->GetName();
+
+	return true;
+}
+
+bool UCollectionManagerScriptingSubsystem::GetCollections(const FCollectionScriptingContainerSource Container, TArray<FCollectionScriptingRef>& OutCollections)
+{
+	const TSharedPtr<ICollectionContainer> CollectionContainer = FCollectionManagerModule::GetModule().Get().FindCollectionContainer(Container.Name);
+	if (!CollectionContainer.IsValid())
+	{
+		CollectionManagerScriptingSubsystemUtil::LogLastCollectionManagerError(FText::Format(CollectionManagerScriptingSubsystemUtil::FConstants::Get().ContainerNotFound, { FText::FromName(Container.Name) }));
+		return false;
+	}
+
+	TArray<FCollectionNameType> Collections;
+	CollectionContainer->GetCollections(Collections);
+
+	OutCollections.Reserve(Collections.Num());
+	Algo::Transform(Collections, OutCollections, [ContainerName = CollectionContainer->GetCollectionSource()->GetName()](const FCollectionNameType& CollectionName)
+		{
+			return FCollectionScriptingRef {
+				.Container = ContainerName,
+				.Name = CollectionName.Name,
+				.ShareType = CollectionManagerScriptingSubsystemUtil::ECollectionShareType_To_ECollectionScriptingShareType(CollectionName.Type)
+			};
+		});
+
+	return true;
+}
+
+bool UCollectionManagerScriptingSubsystem::DestroyCollection(const FCollectionScriptingRef& Collection)
+{
+	const TSharedPtr<ICollectionContainer> CollectionContainer = FCollectionManagerModule::GetModule().Get().FindCollectionContainer(Collection.Container);
+	if (!CollectionContainer.IsValid())
+	{
+		CollectionManagerScriptingSubsystemUtil::LogLastCollectionManagerError(FText::Format(CollectionManagerScriptingSubsystemUtil::FConstants::Get().ContainerNotFound, { FText::FromName(Collection.Container) }));
+		return false;
+	}
+
+	FText Error;
+	if (!CollectionContainer->DestroyCollection(Collection.Name, CollectionManagerScriptingSubsystemUtil::ECollectionScriptingShareType_To_ECollectionShareType(Collection.ShareType), &Error))
+	{
+		CollectionManagerScriptingSubsystemUtil::LogLastCollectionManagerError(Error);
+		return false;
+	}
+
+	return true;
+}
+
+bool UCollectionManagerScriptingSubsystem::RenameCollection(const FCollectionScriptingRef& Collection, const FName NewName, const ECollectionScriptingShareType NewShareType)
+{
+	const TSharedPtr<ICollectionContainer> CollectionContainer = FCollectionManagerModule::GetModule().Get().FindCollectionContainer(Collection.Container);
+	if (!CollectionContainer.IsValid())
+	{
+		CollectionManagerScriptingSubsystemUtil::LogLastCollectionManagerError(FText::Format(CollectionManagerScriptingSubsystemUtil::FConstants::Get().ContainerNotFound, { FText::FromName(Collection.Container) }));
+		return false;
+	}
+
+	FText Error;
+	if (!CollectionContainer->IsValidCollectionName(NewName.ToString(), ECollectionShareType::CST_All, &Error))
+	{
+		CollectionManagerScriptingSubsystemUtil::LogLastCollectionManagerError(Error);
+		return false;
+	}
+
+	if (!CollectionContainer->RenameCollection(Collection.Name, CollectionManagerScriptingSubsystemUtil::ECollectionScriptingShareType_To_ECollectionShareType(Collection.ShareType),
+		NewName, CollectionManagerScriptingSubsystemUtil::ECollectionScriptingShareType_To_ECollectionShareType(NewShareType), &Error))
+	{
+		CollectionManagerScriptingSubsystemUtil::LogLastCollectionManagerError(Error);
+		return false;
+	}
+
+	return true;
+}
+
+bool UCollectionManagerScriptingSubsystem::ReparentCollection(const FCollectionScriptingRef& Collection, const FCollectionScriptingRef NewParentCollection)
+{
+	const TSharedPtr<ICollectionContainer> CollectionContainer = FCollectionManagerModule::GetModule().Get().FindCollectionContainer(Collection.Container);
+	if (!CollectionContainer.IsValid())
+	{
+		CollectionManagerScriptingSubsystemUtil::LogLastCollectionManagerError(FText::Format(CollectionManagerScriptingSubsystemUtil::FConstants::Get().ContainerNotFound, { FText::FromName(Collection.Container) }));
+		return false;
+	}
+
+	// If the intention is to reparent to a different collection, make sure its in the same container.
+	if (!NewParentCollection.Name.IsNone() && Collection.Container != NewParentCollection.Container)
+	{
+		CollectionManagerScriptingSubsystemUtil::LogLastCollectionManagerError(FText::Format(CollectionManagerScriptingSubsystemUtil::FConstants::Get().ContainerMismatch, { FText::FromName(Collection.Container), FText::FromName(NewParentCollection.Container) }));
+		return false;
+	}
+
+	FText Error;
+	if (!CollectionContainer->ReparentCollection(Collection.Name, CollectionManagerScriptingSubsystemUtil::ECollectionScriptingShareType_To_ECollectionShareType(Collection.ShareType),
+		NewParentCollection.Name, CollectionManagerScriptingSubsystemUtil::ECollectionScriptingShareType_To_ECollectionShareType(NewParentCollection.ShareType), &Error))
+	{
+		CollectionManagerScriptingSubsystemUtil::LogLastCollectionManagerError(Error);
+		return false;
+	}
+
+	return true;
+}
+
+bool UCollectionManagerScriptingSubsystem::EmptyCollection(const FCollectionScriptingRef& Collection)
+{
+	const TSharedPtr<ICollectionContainer> CollectionContainer = FCollectionManagerModule::GetModule().Get().FindCollectionContainer(Collection.Container);
+	if (!CollectionContainer.IsValid())
+	{
+		CollectionManagerScriptingSubsystemUtil::LogLastCollectionManagerError(FText::Format(CollectionManagerScriptingSubsystemUtil::FConstants::Get().ContainerNotFound, { FText::FromName(Collection.Container) }));
+		return false;
+	}
+
+	FText Error;
+	if (!CollectionContainer->EmptyCollection(Collection.Name, CollectionManagerScriptingSubsystemUtil::ECollectionScriptingShareType_To_ECollectionShareType(Collection.ShareType), &Error))
+	{
+		CollectionManagerScriptingSubsystemUtil::LogLastCollectionManagerError(Error);
+		return false;
+	}
+
+	return true;
+}
+
+bool UCollectionManagerScriptingSubsystem::AddAssetToCollection(const FCollectionScriptingRef& Collection, const FSoftObjectPath& AssetPath)
+{
+	return AddAssetsToCollection(Collection, TArray<FSoftObjectPath>({ AssetPath }));
+}
+
+bool UCollectionManagerScriptingSubsystem::AddAssetDataToCollection(const FCollectionScriptingRef& Collection, const FAssetData& AssetData)
+{
+	return AddAssetsToCollection(Collection, TArray<FSoftObjectPath>({ AssetData.GetSoftObjectPath() }));
+}
+
+bool UCollectionManagerScriptingSubsystem::AddAssetPtrToCollection(const FCollectionScriptingRef& Collection, const UObject* AssetPtr)
+{
+	return AddAssetsToCollection(Collection, TArray<FSoftObjectPath>({ FSoftObjectPath(AssetPtr) }));
+}
+
+bool UCollectionManagerScriptingSubsystem::AddAssetsToCollection(const FCollectionScriptingRef& Collection, const TArray<FSoftObjectPath>& AssetPaths)
+{
+	const TSharedPtr<ICollectionContainer> CollectionContainer = FCollectionManagerModule::GetModule().Get().FindCollectionContainer(Collection.Container);
+	if (!CollectionContainer.IsValid())
+	{
+		CollectionManagerScriptingSubsystemUtil::LogLastCollectionManagerError(FText::Format(CollectionManagerScriptingSubsystemUtil::FConstants::Get().ContainerNotFound, { FText::FromName(Collection.Container) }));
+		return false;
+	}
+
+	FText Error;
+	if (!CollectionContainer->AddToCollection(Collection.Name, CollectionManagerScriptingSubsystemUtil::ECollectionScriptingShareType_To_ECollectionShareType(Collection.ShareType), AssetPaths, nullptr, &Error))
+	{
+		CollectionManagerScriptingSubsystemUtil::LogLastCollectionManagerError(Error);
+		return false;
+	}
+
+	return true;
+}
+
+bool UCollectionManagerScriptingSubsystem::AddAssetDatasToCollection(const FCollectionScriptingRef& Collection, const TArray<FAssetData>& AssetDatas)
+{
+	TArray<FSoftObjectPath> ObjectPaths;
+	ObjectPaths.Reserve(AssetDatas.Num());
+	Algo::Transform(AssetDatas, ObjectPaths, [](const FAssetData& AssetData) { return AssetData.GetSoftObjectPath(); });
+
+	return AddAssetsToCollection(Collection, ObjectPaths);
+}
+
+bool UCollectionManagerScriptingSubsystem::AddAssetPtrsToCollection(const FCollectionScriptingRef& Collection, const TArray<UObject*>& AssetPtrs)
+{
+	TArray<FSoftObjectPath> ObjectPaths;
+	ObjectPaths.Reserve(AssetPtrs.Num());
+	Algo::Transform(AssetPtrs, ObjectPaths, [](const UObject* AssetPtr) { return FSoftObjectPath(AssetPtr); });
+
+	return AddAssetsToCollection(Collection, ObjectPaths);
+}
+
+bool UCollectionManagerScriptingSubsystem::RemoveAssetFromCollection(const FCollectionScriptingRef& Collection, const FSoftObjectPath& AssetPath)
+{
+	return RemoveAssetsFromCollection(Collection, TArray<FSoftObjectPath>({ AssetPath }));
+}
+
+bool UCollectionManagerScriptingSubsystem::RemoveAssetDataFromCollection(const FCollectionScriptingRef& Collection, const FAssetData& AssetData)
+{
+	return RemoveAssetsFromCollection(Collection, TArray<FSoftObjectPath>({ AssetData.GetSoftObjectPath() }));
+}
+
+bool UCollectionManagerScriptingSubsystem::RemoveAssetPtrFromCollection(const FCollectionScriptingRef& Collection, const UObject* AssetPtr)
+{
+	return RemoveAssetsFromCollection(Collection, TArray<FSoftObjectPath>({ FSoftObjectPath(AssetPtr) }));
+}
+
+bool UCollectionManagerScriptingSubsystem::RemoveAssetsFromCollection(const FCollectionScriptingRef& Collection, const TArray<FSoftObjectPath>& AssetPaths)
+{
+	const TSharedPtr<ICollectionContainer> CollectionContainer = FCollectionManagerModule::GetModule().Get().FindCollectionContainer(Collection.Container);
+	if (!CollectionContainer.IsValid())
+	{
+		CollectionManagerScriptingSubsystemUtil::LogLastCollectionManagerError(FText::Format(CollectionManagerScriptingSubsystemUtil::FConstants::Get().ContainerNotFound, { FText::FromName(Collection.Container) }));
+		return false;
+	}
+
+	FText Error;
+	if (!CollectionContainer->RemoveFromCollection(Collection.Name, CollectionManagerScriptingSubsystemUtil::ECollectionScriptingShareType_To_ECollectionShareType(Collection.ShareType), AssetPaths, nullptr, &Error))
+	{
+		CollectionManagerScriptingSubsystemUtil::LogLastCollectionManagerError(Error);
+		return false;
+	}
+
+	return true;
+}
+
+bool UCollectionManagerScriptingSubsystem::RemoveAssetDatasFromCollection(const FCollectionScriptingRef& Collection, const TArray<FAssetData>& AssetDatas)
+{
+	TArray<FSoftObjectPath> ObjectPaths;
+	ObjectPaths.Reserve(AssetDatas.Num());
+	Algo::Transform(AssetDatas, ObjectPaths, [](const FAssetData& AssetData) { return AssetData.GetSoftObjectPath(); });
+
+	return RemoveAssetsFromCollection(Collection, ObjectPaths);
+}
+
+bool UCollectionManagerScriptingSubsystem::RemoveAssetPtrsFromCollection(const FCollectionScriptingRef& Collection, const TArray<UObject*>& AssetPtrs)
+{
+	TArray<FSoftObjectPath> ObjectPaths;
+	ObjectPaths.Reserve(AssetPtrs.Num());
+	Algo::Transform(AssetPtrs, ObjectPaths, [](const UObject* AssetPtr) { return FSoftObjectPath(AssetPtr); });
+
+	return RemoveAssetsFromCollection(Collection, ObjectPaths);
+}
+
+bool UCollectionManagerScriptingSubsystem::CollectionExists(const FCollectionScriptingContainerSource Container, const FName Collection, const ECollectionScriptingShareType ShareType)
+{
+	const TSharedPtr<ICollectionContainer> CollectionContainer = FCollectionManagerModule::GetModule().Get().FindCollectionContainer(Container.Name);
+	if (!CollectionContainer.IsValid())
+	{
+		CollectionManagerScriptingSubsystemUtil::LogLastCollectionManagerError(FText::Format(CollectionManagerScriptingSubsystemUtil::FConstants::Get().ContainerNotFound, { FText::FromName(Container.Name) }));
+		return false;
+	}
+
+	return CollectionContainer->CollectionExists(Collection, CollectionManagerScriptingSubsystemUtil::ECollectionScriptingShareType_To_ECollectionShareType(ShareType));
+}
+
+bool UCollectionManagerScriptingSubsystem::CollectionExistsWithAnyShareType(const FCollectionScriptingContainerSource Container, const FName Collection)
+{
+	const TSharedPtr<ICollectionContainer> CollectionContainer = FCollectionManagerModule::GetModule().Get().FindCollectionContainer(Container.Name);
+	if (!CollectionContainer.IsValid())
+	{
+		CollectionManagerScriptingSubsystemUtil::LogLastCollectionManagerError(FText::Format(CollectionManagerScriptingSubsystemUtil::FConstants::Get().ContainerNotFound, { FText::FromName(Container.Name) }));
+		return false;
+	}
+
+	return CollectionContainer->CollectionExists(Collection, ECollectionShareType::CST_All);
+}
+
+bool UCollectionManagerScriptingSubsystem::GetAssetsInCollection(const FCollectionScriptingRef& Collection, TArray<FAssetData>& OutAssets)
+{
+	const TSharedPtr<ICollectionContainer> CollectionContainer = FCollectionManagerModule::GetModule().Get().FindCollectionContainer(Collection.Container);
+	if (!CollectionContainer.IsValid())
+	{
+		CollectionManagerScriptingSubsystemUtil::LogLastCollectionManagerError(FText::Format(CollectionManagerScriptingSubsystemUtil::FConstants::Get().ContainerNotFound, { FText::FromName(Collection.Container) }));
+		return false;
+	}
+
+	if (!CollectionContainer->CollectionExists(Collection.Name, CollectionManagerScriptingSubsystemUtil::ECollectionScriptingShareType_To_ECollectionShareType(Collection.ShareType)))
+	{
+		CollectionManagerScriptingSubsystemUtil::LogLastCollectionManagerError(FText::Format(CollectionManagerScriptingSubsystemUtil::FConstants::Get().CollectionNotFound,
+			{ FText::FromName(Collection.Name), FText::FromString(CollectionManagerScriptingSubsystemUtil::ECollectionScriptingShareType_To_String(Collection.ShareType)), FText::FromName(Collection.Container) }));
+		return false;
+	}
+
+	TArray<FSoftObjectPath> Assets;
+	CollectionContainer->GetAssetsInCollection(Collection.Name, CollectionManagerScriptingSubsystemUtil::ECollectionScriptingShareType_To_ECollectionShareType(Collection.ShareType), Assets);
+
+	const IAssetRegistry& AssetRegistry = IAssetRegistry::GetChecked();
+	OutAssets.Reserve(Assets.Num());
+	for (const FSoftObjectPath& Asset : Assets)
+	{
+		FAssetData AssetData = AssetRegistry.GetAssetByObjectPath(Asset);
+		if (AssetData.IsValid())
+		{
+			OutAssets.Emplace(MoveTemp(AssetData));
+		}
+		else
+		{
+			CollectionManagerScriptingSubsystemUtil::LogLastCollectionManagerError(FText::Format(CollectionManagerScriptingSubsystemUtil::FConstants::Get().FailedToGetAsset, { FText::FromString(Asset.ToString()), FText::FromName(Collection.Name) }));
+		}
+	}
+
+	return true;
+}
+
+bool UCollectionManagerScriptingSubsystem::GetCollectionsContainingAsset(const FCollectionScriptingContainerSource Container, const FSoftObjectPath& AssetPath, TArray<FCollectionScriptingRef>& OutCollections)
+{
+	const TSharedPtr<ICollectionContainer> CollectionContainer = FCollectionManagerModule::GetModule().Get().FindCollectionContainer(Container.Name);
+	if (!CollectionContainer.IsValid())
+	{
+		CollectionManagerScriptingSubsystemUtil::LogLastCollectionManagerError(FText::Format(CollectionManagerScriptingSubsystemUtil::FConstants::Get().ContainerNotFound, { FText::FromName(Container.Name) }));
+		return false;
+	}
+
+	TArray<FCollectionNameType> CollectionNamesAndTypes;
+	CollectionContainer->GetCollectionsContainingObject(AssetPath, CollectionNamesAndTypes);
+
+	OutCollections.Reserve(CollectionNamesAndTypes.Num());
+	for (const FCollectionNameType& CollectionNameAndType : CollectionNamesAndTypes)
+	{
+		OutCollections.Emplace(FCollectionScriptingRef{ .Container = Container.Name, .Name = CollectionNameAndType.Name, .ShareType = CollectionManagerScriptingSubsystemUtil::ECollectionShareType_To_ECollectionScriptingShareType(CollectionNameAndType.Type) });
+	}
+
+	return true;
+}
+
+bool UCollectionManagerScriptingSubsystem::GetCollectionsContainingAssetData(const FCollectionScriptingContainerSource Container, const FAssetData& AssetData, TArray<FCollectionScriptingRef>& OutCollections)
+{
+	return GetCollectionsContainingAsset(Container, AssetData.GetSoftObjectPath(), OutCollections);
+}
+
+bool UCollectionManagerScriptingSubsystem::GetCollectionsContainingAssetPtr(const FCollectionScriptingContainerSource Container, const UObject* AssetPtr, TArray<FCollectionScriptingRef>& OutCollections)
+{
+	return GetCollectionsContainingAsset(Container, FSoftObjectPath(AssetPtr), OutCollections);
+}
+
+FCollectionScriptingContainerSource UCollectionManagerScriptingSubsystem::GetBaseGameCollectionContainer() const
+{
+	FCollectionScriptingContainerSource Container;
+	const TSharedRef<ICollectionContainer> CollectionContainer = FCollectionManagerModule::GetModule().Get().GetProjectCollectionContainer();
+	Container.Name = CollectionContainer->GetCollectionSource()->GetName();
+	Container.Title = CollectionContainer->GetCollectionSource()->GetTitle();
+
+	return Container;
+}
+
+#undef LOCTEXT_NAMESPACE
diff --git a/Engine/Source/Editor/UnrealEd/Public/Subsystems/CollectionManagerScriptingSubsystem.h b/Engine/Source/Editor/UnrealEd/Public/Subsystems/CollectionManagerScriptingSubsystem.h
new file mode 100644
--- /dev/null
+++ b/Engine/Source/Editor/UnrealEd/Public/Subsystems/CollectionManagerScriptingSubsystem.h
@@ -0,0 +1,333 @@
+// Copyright Epic Games, Inc. All Rights Reserved.
+
+#pragma once
+
+#include "CollectionManagerScriptingTypes.h"
+#include "EditorSubsystem.h"
+#include "UObject/Package.h"
+
+#include "CollectionManagerScriptingSubsystem.generated.h"
+
+USTRUCT(BlueprintType, DisplayName="Collection Container Source", meta=(ScriptName="CollectionContainerSource"))
+struct FCollectionScriptingContainerSource
+{
+	GENERATED_BODY()
+
+	/** The name of the container. Defaults to the base game's container. */
+	UPROPERTY(BlueprintReadOnly, Category="CollectionManager")
+	FName Name = NAME_Game;
+
+	/** Friendly title of the container. */
+	UPROPERTY(BlueprintReadOnly, Category="CollectionManager")
+	FText Title = NSLOCTEXT("CollectionScriptingContainerSource", "ProjectCollectionSource_Name", "Collections");
+};
+
+USTRUCT(BlueprintType, DisplayName="Collection", meta=(ScriptName="Collection"))
+struct FCollectionScriptingRef
+{
+	GENERATED_BODY()
+
+	/** The name (not title) of the container that owns this collection. Defaults to the base game's container. */
+	UPROPERTY(BlueprintReadWrite, Category="CollectionManager")
+	FName Container = NAME_Game;
+
+	/** Friendly name of the collection. */
+	UPROPERTY(BlueprintReadWrite, Category="CollectionManager")
+	FName Name;
+
+	/** Share type of this collection, which controls its visibility. */
+	UPROPERTY(BlueprintReadWrite, Category="CollectionManager")
+	ECollectionScriptingShareType ShareType = ECollectionScriptingShareType::Local;
+};
+
+UCLASS(DisplayName="Collection Manager Subsystem", meta=(ScriptName="CollectionManagerSubsystem"))
+class UNREALED_API UCollectionManagerScriptingSubsystem final : public UEditorSubsystem
+{
+	GENERATED_BODY()
+
+public:
+	/**
+	 * Gets all available collection containers.
+	 *
+	 * @return Array of all collection containers.
+	 */
+	UFUNCTION(BlueprintCallable, Category="CollectionManager")
+	TArray<FCollectionScriptingContainerSource> GetCollectionContainers();
+
+	/**
+	 * Create a new collection with the given name and share type in the provided collection container.
+	 *
+	 * @param Container The container the collection should be created in. "None" defaults to the base game's container.
+	 * @param Collection Name to give to the collection.
+	 * @param ShareType Whether the collection should be local, private, or shared?
+	 * @param OutNewCollection The newly created collection.
+	 *
+	 * @return True if the collection was created, false otherwise (see the output log for details on error).
+	 */
+	UFUNCTION(BlueprintCallable, Category="CollectionManager")
+	bool CreateCollection(const FCollectionScriptingContainerSource Container, const FName Collection, const ECollectionScriptingShareType ShareType, FCollectionScriptingRef& OutNewCollection);
+
+	/**
+	 * Get all available collections in the specified container.
+	 *
+	 * @param Container The container to retrieve collections from. "None" defaults to the base game's container.
+	 * @param OutCollections The collections found in the container.
+	 *
+	 * @return True if the specified container was found, false otherwise (see the output log for details on error).
+	 */
+	UFUNCTION(BlueprintCallable, Category="CollectionManager")
+	bool GetCollections(const FCollectionScriptingContainerSource Container, TArray<FCollectionScriptingRef>& OutCollections);
+
+	/**
+	 * Destroy the given collection.
+	 *
+	 * @param Collection The collection to destroy.
+	 *
+	 * @return True if the collection was destroyed, false otherwise (see the output log for details on error).
+	 */
+	UFUNCTION(BlueprintCallable, Category="CollectionManager")
+	bool DestroyCollection(const FCollectionScriptingRef& Collection);
+
+	/**
+	 * Rename the given collection.
+	 *
+	 * @param Collection The collection to rename.
+	 * @param NewName New name to give to the collection.
+	 * @param NewShareType New share type to give to the collection (local, private, or shared).
+	 *
+	 * @return True if the collection was renamed, false otherwise (see the output log for details on error).
+	 */
+	UFUNCTION(BlueprintCallable, Category="CollectionManager")
+	bool RenameCollection(const FCollectionScriptingRef& Collection, const FName NewName, const ECollectionScriptingShareType NewShareType);
+
+	/**
+	 * Re-parent the given collection.
+	 *
+	 * @param Collection The collection to re-parent.
+	 * @param NewParentCollection The new parent collection, or "None" to have the collection become a root collection.
+	 *
+	 * @return True if the collection was re-parented, false otherwise (see the output log for details on error).
+	 */
+	UFUNCTION(BlueprintCallable, Category="CollectionManager")
+	bool ReparentCollection(const FCollectionScriptingRef& Collection, const FCollectionScriptingRef NewParentCollection);
+
+	/**
+	 * Remove all assets from the given collection.
+	 *
+	 * @param Collection The collection to modify.
+	 *
+	 * @return True if the collection was modified, false otherwise (see the output log for details on error).
+	 */
+	UFUNCTION(BlueprintCallable, Category="CollectionManager")
+	bool EmptyCollection(const FCollectionScriptingRef& Collection);
+
+	/**
+	 * Add the given asset to the given collection.
+	 *
+	 * @param Collection The collection to modify.
+	 * @param AssetPath Asset to add (its path name, eg) /Game/MyFolder/MyAsset.MyAsset).
+	 *
+	 * @return True if the collection was modified, false otherwise (see the output log for details on error).
+	 */
+	UFUNCTION(BlueprintCallable, Category="CollectionManager")
+	bool AddAssetToCollection(const FCollectionScriptingRef& Collection, const FSoftObjectPath& AssetPath);
+
+	/**
+	 * Add the given asset to the given collection.
+	 *
+	 * @param Collection The collection to modify.
+	 * @param AssetData Asset to add.
+	 *
+	 * @return True if the collection was modified, false otherwise (see the output log for details on error).
+	 */
+	UFUNCTION(BlueprintCallable, Category="CollectionManager")
+	bool AddAssetDataToCollection(const FCollectionScriptingRef& Collection, const FAssetData& AssetData);
+
+	/**
+	 * Add the given asset to the given collection.
+	 *
+	 * @param Collection The collection to modify.
+	 * @param AssetPtr Asset to add.
+	 *
+	 * @return True if the collection was modified, false otherwise (see the output log for details on error).
+	 */
+	UFUNCTION(BlueprintCallable, Category="CollectionManager")
+	bool AddAssetPtrToCollection(const FCollectionScriptingRef& Collection, const UObject* AssetPtr);
+
+	/**
+	 * Add the given assets to the given collection.
+	 *
+	 * @param Collection The collection to modify.
+	 * @param AssetPaths Assets to add (their paths, eg) /Game/MyFolder/MyAsset.MyAsset).
+	 *
+	 * @return True if the collection was modified, false otherwise (see the output log for details on error).
+	 */
+	UFUNCTION(BlueprintCallable, Category="CollectionManager")
+	bool AddAssetsToCollection(const FCollectionScriptingRef& Collection, const TArray<FSoftObjectPath>& AssetPaths);
+
+	/**
+	 * Add the given assets to the given collection.
+	 *
+	 * @param Collection The collection to modify.
+	 * @param AssetDatas Assets to add.
+	 *
+	 * @return True if the collection was modified, false otherwise (see the output log for details on error).
+	 */
+	UFUNCTION(BlueprintCallable, Category="CollectionManager")
+	bool AddAssetDatasToCollection(const FCollectionScriptingRef& Collection, const TArray<FAssetData>& AssetDatas);
+
+	/**
+	 * Add the given assets to the given collection.
+	 *
+	 * @param Collection The collection to modify.
+	 * @param AssetPtrs Assets to add.
+	 *
+	 * @return True if the collection was modified, false otherwise (see the output log for details on error).
+	 */
+	UFUNCTION(BlueprintCallable, Category="CollectionManager")
+	bool AddAssetPtrsToCollection(const FCollectionScriptingRef& Collection, const TArray<UObject*>& AssetPtrs);
+
+	/**
+	 * Remove the given asset from the given collection.
+	 *
+	 * @param Collection The collection to modify.
+	 * @param AssetPath Asset to remove (its path, eg) /Game/MyFolder/MyAsset.MyAsset).
+	 *
+	 * @return True if the collection was modified, false otherwise (see the output log for details on error).
+	 */
+	UFUNCTION(BlueprintCallable, Category="CollectionManager")
+	bool RemoveAssetFromCollection(const FCollectionScriptingRef& Collection, const FSoftObjectPath& AssetPath);
+
+	/**
+	 * Remove the given asset from the given collection.
+	 *
+	 * @param Collection The collection to modify.
+	 * @param AssetData Asset to remove.
+	 *
+	 * @return True if the collection was modified, false otherwise (see the output log for details on error).
+	 */
+	UFUNCTION(BlueprintCallable, Category="CollectionManager")
+	bool RemoveAssetDataFromCollection(const FCollectionScriptingRef& Collection, const FAssetData& AssetData);
+
+	/**
+	 * Remove the given asset from the given collection.
+	 *
+	 * @param Collection The collection to modify.
+	 * @param AssetPtr Asset to remove.
+	 *
+	 * @return True if the collection was modified, false otherwise (see the output log for details on error).
+	 */
+	UFUNCTION(BlueprintCallable, Category="CollectionManager")
+	bool RemoveAssetPtrFromCollection(const FCollectionScriptingRef& Collection, const UObject* AssetPtr);
+
+	/**
+	 * Remove the given assets from the given collection.
+	 *
+	 * @param Collection The collection to modify.
+	 * @param AssetPaths Assets to remove (their paths, eg) /Game/MyFolder/MyAsset.MyAsset).
+	 *
+	 * @return True if the collection was modified, false otherwise (see the output log for details on error).
+	 */
+	UFUNCTION(BlueprintCallable, Category="CollectionManager")
+	bool RemoveAssetsFromCollection(const FCollectionScriptingRef& Collection, const TArray<FSoftObjectPath>& AssetPaths);
+
+	/**
+	 * Remove the given assets from the given collection.
+	 *
+	 * @param Collection The collection to modify.
+	 * @param AssetDatas Assets to remove.
+	 *
+	 * @return True if the collection was modified, false otherwise (see the output log for details on error).
+	 */
+	UFUNCTION(BlueprintCallable, Category="CollectionManager")
+	bool RemoveAssetDatasFromCollection(const FCollectionScriptingRef& Collection, const TArray<FAssetData>& AssetDatas);
+
+	/**
+	 * Remove the given assets from the given collection.
+	 *
+	 * @param Collection The collection to modify.
+	 * @param AssetPtrs Assets to remove.
+	 *
+	 * @return True if the collection was modified, false otherwise (see the output log for details on error).
+	 */
+	UFUNCTION(BlueprintCallable, Category="CollectionManager")
+	bool RemoveAssetPtrsFromCollection(const FCollectionScriptingRef& Collection, const TArray<UObject*>& AssetPtrs);
+
+	/**
+	 * Check whether the given collection exists in the given container (matching both name and share type).
+	 *
+	 * @param Container The container to search. "None" defaults to the base game's container.
+	 * @param Collection The name of the collection to look for.
+	 * @param ShareType The share type of the collection to look for.
+	 *
+	 * @return True if the collection exists, false otherwise.
+	 */
+	UFUNCTION(BlueprintCallable, Category="CollectionManager")
+	bool CollectionExists(const FCollectionScriptingContainerSource Container, const FName Collection, const ECollectionScriptingShareType ShareType);
+
+	/**
+	 * Check whether the given collection exists in the given container (matching only by name).
+	 *
+	 * @param Container The container to search. "None" defaults to the base game's container.
+	 * @param Collection The collection to look for. 
+	 *
+	 * @return True if the collection exists (returns the first found with a matching name), false otherwise.
+	 */
+	UFUNCTION(BlueprintCallable, Category = "CollectionManager")
+	bool CollectionExistsWithAnyShareType(const FCollectionScriptingContainerSource Container, const FName Collection);
+
+	/**
+	 * Get the assets in the given collection.
+	 *
+	 * @param Collection The collection from which to retrieve assets.
+	 * @param OutAssets The assets found in the collection.
+	 *
+	 * @return True if the collection exists, false otherwise (see the output log for details on error).
+	 */
+	UFUNCTION(BlueprintCallable, Category="CollectionManager")
+	bool GetAssetsInCollection(const FCollectionScriptingRef& Collection, TArray<FAssetData>& OutAssets);
+
+	/**
+	 * Get the collections in the specified container that contain the given asset.
+	 *
+	 * @param Container Container to search for collections containing the given asset. "None" defaults to the base game's container.
+	 * @param AssetPath Asset to test (its path name, eg) /Game/MyFolder/MyAsset.MyAsset).
+	 * @param OutCollections Array of the collections that contain the asset.
+	 *
+	 * @return True if the container exists, false otherwise (see the output log for details on error).
+	 */
+	UFUNCTION(BlueprintCallable, Category="CollectionManager")
+	bool GetCollectionsContainingAsset(const FCollectionScriptingContainerSource Container, const FSoftObjectPath& AssetPath, TArray<FCollectionScriptingRef>& OutCollections);
+
+	/**
+	 * Get the collections in the specified container that contain the given asset.
+	 *
+	 * @param Container Container to search for collections containing the given asset. "None" defaults to the base game's container.
+	 * @param AssetData Asset to test.
+	 * @param OutCollections Array of the collections that contain the asset.
+	 *
+	 * @return True if the container exists, false otherwise (see the output log for details on error).
+	 */
+	UFUNCTION(BlueprintCallable, Category="CollectionManager")
+	bool GetCollectionsContainingAssetData(const FCollectionScriptingContainerSource Container, const FAssetData& AssetData, TArray<FCollectionScriptingRef>& OutCollections);
+
+	/**
+	 * Get the collections in the specified container that contain the given asset.
+	 *
+	 * @param Container Container to search for collections containing the given asset. "None" defaults to the base game's container.
+	 * @param AssetPtr Asset to test.
+	 * @param OutCollections Array of the collections that contain the asset.
+	 *
+	 * @return True if the container exists, false otherwise (see the output log for details on error).
+	 */
+	UFUNCTION(BlueprintCallable, Category="CollectionManager")
+	bool GetCollectionsContainingAssetPtr(const FCollectionScriptingContainerSource Container, const UObject* AssetPtr, TArray<FCollectionScriptingRef>& OutCollections);
+
+	/**
+	 * Get the collection container for the base game.
+	 *
+	 * @return The collection container for the base game.
+	 */
+	UFUNCTION(BlueprintCallable, Category="CollectionManager")
+	FCollectionScriptingContainerSource GetBaseGameCollectionContainer() const;
+};
diff --git a/Engine/Source/Runtime/Engine/Public/CollectionManagerScriptingTypes.h b/Engine/Source/Runtime/Engine/Public/CollectionManagerScriptingTypes.h
new file mode 100644
--- /dev/null
+++ b/Engine/Source/Runtime/Engine/Public/CollectionManagerScriptingTypes.h
@@ -0,0 +1,16 @@
+// Copyright Epic Games, Inc. All Rights Reserved.
+
+#pragma once 
+
+#include "CollectionManagerScriptingTypes.generated.h"
+
+UENUM(DisplayName="Collection Share Type", meta=(ScriptName="CollectionShareType"))
+enum class ECollectionScriptingShareType : uint8
+{
+	/** This collection is only visible to you and is not in source control. */
+	Local,
+	/** This collection is only visible to you. */
+	Private,
+	/** This collection is visible to everyone. */
+	Shared,
+};
