Deprecating the AssetTagsSubsystem's editor-only Collection BP APIs in favor of the new CollectionManagerScriptingSubsystem BP API.

[RN] Marked editor-only functions as deprecated with an encouragement to migrate to Collection Manager Scripting Subsystem

diff --git a/Engine/Plugins/Runtime/AssetTags/Source/AssetTags/Public/AssetTagsSubsystem.h b/Engine/Plugins/Runtime/AssetTags/Source/AssetTags/Public/AssetTagsSubsystem.h
--- a/Engine/Plugins/Runtime/AssetTags/Source/AssetTags/Public/AssetTagsSubsystem.h
+++ b/Engine/Plugins/Runtime/AssetTags/Source/AssetTags/Public/AssetTagsSubsystem.h
@@ -22,7 +22,8 @@ public:
 	 *
 	 * @return True if the collection was created, false otherwise (see the output log for details on error).
 	 */
-	UFUNCTION(BlueprintCallable, Category="AssetTags")
+	UE_DEPRECATED(5.6, "Collection Manager Scripting Subsystem has been introduced as a new top level object for managing collections. Use that instead.")
+	UFUNCTION(BlueprintCallable, Category="AssetTags", meta=(DeprecatedFunction, DeprecationMessage="Collection Manager Scripting Subsystem has been introduced as a new top level object for managing collections. Use that instead."))
 	bool CreateCollection(const FName Name, const ECollectionScriptingShareType ShareType);
 	
 	/**
@@ -32,7 +33,8 @@ public:
 	 *
 	 * @return True if the collection was destroyed, false otherwise (see the output log for details on error).
 	 */
-	UFUNCTION(BlueprintCallable, Category="AssetTags")
+	UE_DEPRECATED(5.6, "Collection Manager Scripting Subsystem has been introduced as a new top level object for managing collections. Use that instead.")
+	UFUNCTION(BlueprintCallable, Category="AssetTags", meta=(DeprecatedFunction, DeprecationMessage="Collection Manager Scripting Subsystem has been introduced as a new top level object for managing collections. Use that instead."))
 	bool DestroyCollection(const FName Name);
 	
 	/**
@@ -43,7 +45,8 @@ public:
 	 *
 	 * @return True if the collection was renamed, false otherwise (see the output log for details on error).
 	 */
-	UFUNCTION(BlueprintCallable, Category="AssetTags")
+	UE_DEPRECATED(5.6, "Collection Manager Scripting Subsystem has been introduced as a new top level object for managing collections. Use that instead.")
+	UFUNCTION(BlueprintCallable, Category="AssetTags", meta=(DeprecatedFunction, DeprecationMessage="Collection Manager Scripting Subsystem has been introduced as a new top level object for managing collections. Use that instead."))
 	bool RenameCollection(const FName Name, const FName NewName);
 	
 	/**
@@ -54,7 +57,8 @@ public:
 	 *
 	 * @return True if the collection was renamed, false otherwise (see the output log for details on error).
 	 */
-	UFUNCTION(BlueprintCallable, Category="AssetTags")
+	UE_DEPRECATED(5.6, "Collection Manager Scripting Subsystem has been introduced as a new top level object for managing collections. Use that instead.")
+	UFUNCTION(BlueprintCallable, Category="AssetTags", meta=(DeprecatedFunction, DeprecationMessage="Collection Manager Scripting Subsystem has been introduced as a new top level object for managing collections. Use that instead."))
 	bool ReparentCollection(const FName Name, const FName NewParentName);
 
 	/**
@@ -64,7 +68,8 @@ public:
 	 *
 	 * @return True if the collection was modified, false otherwise (see the output log for details on error).
 	 */
-	UFUNCTION(BlueprintCallable, Category="AssetTags")
+	UE_DEPRECATED(5.6, "Collection Manager Scripting Subsystem has been introduced as a new top level object for managing collections. Use that instead.")
+	UFUNCTION(BlueprintCallable, Category="AssetTags", meta=(DeprecatedFunction, DeprecationMessage="Collection Manager Scripting Subsystem has been introduced as a new top level object for managing collections. Use that instead."))
 	bool EmptyCollection(const FName Name);
 
 	/**
@@ -75,7 +80,8 @@ public:
 	 *
 	 * @return True if the collection was modified, false otherwise (see the output log for details on error).
 	 */
-	UFUNCTION(BlueprintCallable, Category="AssetTags", DisplayName="Add Asset To Collection")
+	UE_DEPRECATED(5.6, "Collection Manager Scripting Subsystem has been introduced as a new top level object for managing collections. Use that instead.")
+	UFUNCTION(BlueprintCallable, Category="AssetTags", DisplayName="Add Asset To Collection", meta=(DeprecatedFunction, DeprecationMessage="Collection Manager Scripting Subsystem has been introduced as a new top level object for managing collections. Use that instead."))
 	bool K2_AddAssetToCollection(const FName Name, const FSoftObjectPath& AssetPath);
 
 	UE_DEPRECATED(5.1, "Names containing full asset paths are deprecated. Use Soft Object Path instead.")
@@ -90,7 +96,8 @@ public:
 	 *
 	 * @return True if the collection was modified, false otherwise (see the output log for details on error).
 	 */
-	UFUNCTION(BlueprintCallable, Category="AssetTags")
+	UE_DEPRECATED(5.6, "Collection Manager Scripting Subsystem has been introduced as a new top level object for managing collections. Use that instead.")
+	UFUNCTION(BlueprintCallable, Category="AssetTags", meta=(DeprecatedFunction, DeprecationMessage="Collection Manager Scripting Subsystem has been introduced as a new top level object for managing collections. Use that instead."))
 	bool AddAssetDataToCollection(const FName Name, const FAssetData& AssetData);
 
 	/**
@@ -101,7 +108,8 @@ public:
 	 *
 	 * @return True if the collection was modified, false otherwise (see the output log for details on error).
 	 */
-	UFUNCTION(BlueprintCallable, Category="AssetTags")
+	UE_DEPRECATED(5.6, "Collection Manager Scripting Subsystem has been introduced as a new top level object for managing collections. Use that instead.")
+	UFUNCTION(BlueprintCallable, Category="AssetTags", meta=(DeprecatedFunction, DeprecationMessage="Collection Manager Scripting Subsystem has been introduced as a new top level object for managing collections. Use that instead."))
 	bool AddAssetPtrToCollection(const FName Name, const UObject* AssetPtr);
 	
 	/**
@@ -112,7 +120,8 @@ public:
 	 *
 	 * @return True if the collection was modified, false otherwise (see the output log for details on error).
 	 */
-	UFUNCTION(BlueprintCallable, Category="AssetTags", DisplayName="Add Assets To Collection")
+	UE_DEPRECATED(5.6, "Collection Manager Scripting Subsystem has been introduced as a new top level object for managing collections. Use that instead.")
+	UFUNCTION(BlueprintCallable, Category="AssetTags", DisplayName="Add Assets To Collection", meta=(DeprecatedFunction, DeprecationMessage="Collection Manager Scripting Subsystem has been introduced as a new top level object for managing collections. Use that instead."))
 	bool K2_AddAssetsToCollection(const FName Name, const TArray<FSoftObjectPath>& AssetPaths);
 	 
 	UE_DEPRECATED(5.1, "Names containing full asset paths are deprecated. Use Soft Object Path instead.")
@@ -127,7 +136,8 @@ public:
 	 *
 	 * @return True if the collection was modified, false otherwise (see the output log for details on error).
 	 */
-	UFUNCTION(BlueprintCallable, Category="AssetTags")
+	UE_DEPRECATED(5.6, "Collection Manager Scripting Subsystem has been introduced as a new top level object for managing collections. Use that instead.")
+	UFUNCTION(BlueprintCallable, Category="AssetTags", meta=(DeprecatedFunction, DeprecationMessage="Collection Manager Scripting Subsystem has been introduced as a new top level object for managing collections. Use that instead."))
 	bool AddAssetDatasToCollection(const FName Name, const TArray<FAssetData>& AssetDatas);
 
 	/**
@@ -138,7 +148,8 @@ public:
 	 *
 	 * @return True if the collection was modified, false otherwise (see the output log for details on error).
 	 */
-	UFUNCTION(BlueprintCallable, Category="AssetTags")
+	UE_DEPRECATED(5.6, "Collection Manager Scripting Subsystem has been introduced as a new top level object for managing collections. Use that instead.")
+	UFUNCTION(BlueprintCallable, Category="AssetTags", meta=(DeprecatedFunction, DeprecationMessage="Collection Manager Scripting Subsystem has been introduced as a new top level object for managing collections. Use that instead."))
 	bool AddAssetPtrsToCollection(const FName Name, const TArray<UObject*>& AssetPtrs);
 
 	/**
@@ -149,7 +160,8 @@ public:
 	 *
 	 * @return True if the collection was modified, false otherwise (see the output log for details on error).
 	 */
-	UFUNCTION(BlueprintCallable, Category="AssetTags", DisplayName="Remove Asset From Collection")
+	UE_DEPRECATED(5.6, "Collection Manager Scripting Subsystem has been introduced as a new top level object for managing collections. Use that instead.")
+	UFUNCTION(BlueprintCallable, Category="AssetTags", DisplayName="Remove Asset From Collection", meta=(DeprecatedFunction, DeprecationMessage="Collection Manager Scripting Subsystem has been introduced as a new top level object for managing collections. Use that instead."))
 	bool K2_RemoveAssetFromCollection(const FName Name, const FSoftObjectPath& AssetPath);
 
 	UE_DEPRECATED(5.1, "Names containing full asset paths are deprecated, use Soft Object Path instead.")
@@ -164,7 +176,8 @@ public:
 	 *
 	 * @return True if the collection was modified, false otherwise (see the output log for details on error).
 	 */
-	UFUNCTION(BlueprintCallable, Category="AssetTags")
+	UE_DEPRECATED(5.6, "Collection Manager Scripting Subsystem has been introduced as a new top level object for managing collections. Use that instead.")
+	UFUNCTION(BlueprintCallable, Category="AssetTags", meta=(DeprecatedFunction, DeprecationMessage="Collection Manager Scripting Subsystem has been introduced as a new top level object for managing collections. Use that instead."))
 	bool RemoveAssetDataFromCollection(const FName Name, const FAssetData& AssetData);
 
 	/**
@@ -175,7 +188,8 @@ public:
 	 *
 	 * @return True if the collection was modified, false otherwise (see the output log for details on error).
 	 */
-	UFUNCTION(BlueprintCallable, Category="AssetTags")
+	UE_DEPRECATED(5.6, "Collection Manager Scripting Subsystem has been introduced as a new top level object for managing collections. Use that instead.")
+	UFUNCTION(BlueprintCallable, Category="AssetTags", meta=(DeprecatedFunction, DeprecationMessage="Collection Manager Scripting Subsystem has been introduced as a new top level object for managing collections. Use that instead."))
 	bool RemoveAssetPtrFromCollection(const FName Name, const UObject* AssetPtr);
 	
 	/**
@@ -186,7 +200,8 @@ public:
 	 *
 	 * @return True if the collection was modified, false otherwise (see the output log for details on error).
 	 */
-	UFUNCTION(BlueprintCallable, Category="AssetTags", DisplayName="Remove Assets From Collection")
+	UE_DEPRECATED(5.6, "Collection Manager Scripting Subsystem has been introduced as a new top level object for managing collections. Use that instead.")
+	UFUNCTION(BlueprintCallable, Category="AssetTags", DisplayName="Remove Assets From Collection", meta=(DeprecatedFunction, DeprecationMessage="Collection Manager Scripting Subsystem has been introduced as a new top level object for managing collections. Use that instead."))
 	bool K2_RemoveAssetsFromCollection(const FName Name, const TArray<FSoftObjectPath>& AssetPaths);
 
 	UE_DEPRECATED(5.1, "Names containing full asset paths are deprecated. Use Soft Object Path instead.")
@@ -201,7 +216,8 @@ public:
 	 *
 	 * @return True if the collection was modified, false otherwise (see the output log for details on error).
 	 */
-	UFUNCTION(BlueprintCallable, Category="AssetTags")
+	UE_DEPRECATED(5.6, "Collection Manager Scripting Subsystem has been introduced as a new top level object for managing collections. Use that instead.")
+	UFUNCTION(BlueprintCallable, Category="AssetTags", meta=(DeprecatedFunction, DeprecationMessage="Collection Manager Scripting Subsystem has been introduced as a new top level object for managing collections. Use that instead."))
 	bool RemoveAssetDatasFromCollection(const FName Name, const TArray<FAssetData>& AssetDatas);
 
 	/**
@@ -212,13 +228,15 @@ public:
 	 *
 	 * @return True if the collection was modified, false otherwise (see the output log for details on error).
 	 */
-	UFUNCTION(BlueprintCallable, Category="AssetTags")
+	UE_DEPRECATED(5.6, "Collection Manager Scripting Subsystem has been introduced as a new top level object for managing collections. Use that instead.")
+	UFUNCTION(BlueprintCallable, Category="AssetTags", meta=(DeprecatedFunction, DeprecationMessage="Collection Manager Scripting Subsystem has been introduced as a new top level object for managing collections. Use that instead."))
 	bool RemoveAssetPtrsFromCollection(const FName Name, const TArray<UObject*>& AssetPtrs);
 #endif 	// WITH_EDITOR
 
 	/**
 	 * Check whether the given collection exists.
-	 *
+	 * Use this for in-game access. Use the Collections Manager Scripting Subsystem for any other use case.
+	 * 
 	 * @param Name Name of the collection to test.
 	 *
 	 * @return True if the collection exists, false otherwise.
@@ -228,7 +246,8 @@ public:
 
 	/**
 	 * Get the names of all available collections.
-	 *
+	 * Use this for in-game access. Use the Collections Manager Scripting Subsystem for any other use case.
+	 * 
 	 * @return Names of all available collections.
 	 */
 	UFUNCTION(BlueprintCallable, Category="AssetTags")
@@ -236,6 +255,7 @@ public:
 
 	/**
 	 * Get the assets in the given collection.
+	 * Use this for in-game access. Use the Collections Manager Scripting Subsystem for any other use case.
 	 * 
 	 * @param Name Name of the collection to test.
 	 *
@@ -246,6 +266,7 @@ public:
 
 	/**
 	 * Get the names of the collections that contain the given asset.
+	 * Use this for in-game access. Use the Collections Manager Scripting Subsystem for any other use case.
 	 * 
 	 * @param AssetPathName Asset to test (its path name, eg) /Game/MyFolder/MyAsset.MyAsset).
 	 *
@@ -260,6 +281,7 @@ public:
 
 	/**
 	 * Get the names of the collections that contain the given asset.
+	 * Use this for in-game access. Use the Collections Manager Scripting Subsystem for any other use case.
 	 * 
 	 * @param AssetData Asset to test.
 	 *
@@ -270,6 +292,7 @@ public:
 
 	/**
 	 * Get the names of the collections that contain the given asset.
+	 * Use this for in-game access. Use the Collections Manager Scripting Subsystem for any other use case.
 	 * 
 	 * @param AssetPtr Asset to test.
 	 *
